UE4 -> UnrealEditor

diff --git a/Engine/Plugins/Developer/CodeLiteSourceCodeAccess/Source/CodeLiteSourceCodeAccess/Private/CodeLiteSourceCodeAccessor.cpp b/Engine/Plugins/Developer/CodeLiteSourceCodeAccess/Source/CodeLiteSourceCodeAccess/Private/CodeLiteSourceCodeAccessor.cpp
--- a/Engine/Plugins/Developer/CodeLiteSourceCodeAccess/Source/CodeLiteSourceCodeAccess/Private/CodeLiteSourceCodeAccessor.cpp
+++ b/Engine/Plugins/Developer/CodeLiteSourceCodeAccess/Source/CodeLiteSourceCodeAccess/Private/CodeLiteSourceCodeAccessor.cpp
@@ -178,7 +178,7 @@ FString FCodeLiteSourceCodeAccessor::GetSolutionPath() const
 		
 		if (!FUProjectDictionary::GetDefault().IsForeignProject(CachedSolutionPath))
 		{
-			CachedSolutionPath = FPaths::Combine(FPaths::RootDir(), TEXT("UE4.workspace"));
+			CachedSolutionPath = FPaths::Combine(FPaths::RootDir(), TEXT("UnrealEditor.workspace"));
 		}
 		else
 		{
