PR #4992: [BUGFIX] Codelite gets wrong path to open workspace file

#rnx

diff --git a/Engine/Plugins/Developer/CodeLiteSourceCodeAccess/Source/CodeLiteSourceCodeAccess/CodeLiteSourceCodeAccess.Build.cs b/Engine/Plugins/Developer/CodeLiteSourceCodeAccess/Source/CodeLiteSourceCodeAccess/CodeLiteSourceCodeAccess.Build.cs
--- a/Engine/Plugins/Developer/CodeLiteSourceCodeAccess/Source/CodeLiteSourceCodeAccess/CodeLiteSourceCodeAccess.Build.cs
+++ b/Engine/Plugins/Developer/CodeLiteSourceCodeAccess/Source/CodeLiteSourceCodeAccess/CodeLiteSourceCodeAccess.Build.cs
@@ -12,6 +12,7 @@ namespace UnrealBuildTool.Rules
 					"Core",
 					"SourceCodeAccess",
 					"DesktopPlatform",
+					"Projects",
 				}
 			);
 
diff --git a/Engine/Plugins/Developer/CodeLiteSourceCodeAccess/Source/CodeLiteSourceCodeAccess/Private/CodeLiteSourceCodeAccessModule.cpp b/Engine/Plugins/Developer/CodeLiteSourceCodeAccess/Source/CodeLiteSourceCodeAccess/Private/CodeLiteSourceCodeAccessModule.cpp
--- a/Engine/Plugins/Developer/CodeLiteSourceCodeAccess/Source/CodeLiteSourceCodeAccess/Private/CodeLiteSourceCodeAccessModule.cpp
+++ b/Engine/Plugins/Developer/CodeLiteSourceCodeAccess/Source/CodeLiteSourceCodeAccess/Private/CodeLiteSourceCodeAccessModule.cpp
@@ -6,6 +6,8 @@
 
 IMPLEMENT_MODULE( FCodeLiteSourceCodeAccessModule, CodeLiteSourceCodeAccess );
 
+#define LOCTEXT_NAMESPACE "CodeLiteSourceCodeAccessor"
+
 void FCodeLiteSourceCodeAccessModule::StartupModule()
 {
 	CodeLiteSourceCodeAccessor.Startup();
@@ -24,3 +26,5 @@ FCodeLiteSourceCodeAccessor& FCodeLiteSourceCodeAccessModule::GetAccessor()
 {
 	return CodeLiteSourceCodeAccessor;
 }
+
+#undef LOCTEXT_NAMESPACE
\ No newline at end of file
diff --git a/Engine/Plugins/Developer/CodeLiteSourceCodeAccess/Source/CodeLiteSourceCodeAccess/Private/CodeLiteSourceCodeAccessor.cpp b/Engine/Plugins/Developer/CodeLiteSourceCodeAccess/Source/CodeLiteSourceCodeAccess/Private/CodeLiteSourceCodeAccessor.cpp
--- a/Engine/Plugins/Developer/CodeLiteSourceCodeAccess/Source/CodeLiteSourceCodeAccess/Private/CodeLiteSourceCodeAccessor.cpp
+++ b/Engine/Plugins/Developer/CodeLiteSourceCodeAccess/Source/CodeLiteSourceCodeAccess/Private/CodeLiteSourceCodeAccessor.cpp
@@ -4,18 +4,47 @@
 #include "Misc/Paths.h"
 #include "Misc/UProjectInfo.h"
 #include "Misc/App.h"
+#include "Misc/FileHelper.h"
 #include "Trace/Trace.inl"
+#include "ProjectDescriptor.h"
+#include "Interfaces/IProjectManager.h"
 
 DEFINE_LOG_CATEGORY_STATIC(LogCodeLiteAccessor, Log, All);
 
 #define LOCTEXT_NAMESPACE "CodeLiteSourceCodeAccessor"
 
+static FString MakePath(const FString& InPath)
+{
+	return TEXT("\"") + InPath.ReplaceQuotesWithEscapedQuotes() + TEXT("\""); 
+}
+
 void FCodeLiteSourceCodeAccessor::Startup()
 {
 	// Cache this so we don't have to do it on a background thread
 	GetSolutionPath();
 }
 
+void FCodeLiteSourceCodeAccessor::RefreshAvailability()
+{
+	FString OutURL;
+	int32 ReturnCode = -1;
+
+	FPlatformProcess::ExecProcess(TEXT("/usr/bin/env"), TEXT("bash -c \"type -p codelite\""), &ReturnCode, &OutURL, nullptr);
+	if (ReturnCode == 0)
+	{
+		Location.URL = OutURL.TrimStartAndEnd();
+	}
+	else
+	{
+		// Fallback to default install location
+		FString URL = TEXT("/usr/bin/codelite");
+		if (FPaths::FileExists(URL))
+		{
+			Location.URL = URL;
+		}
+	}
+}
+
 void FCodeLiteSourceCodeAccessor::Shutdown()
 {
 
@@ -23,18 +52,17 @@ void FCodeLiteSourceCodeAccessor::Shutdown()
 
 bool FCodeLiteSourceCodeAccessor::CanAccessSourceCode() const
 {
-	FString Path;
-	return CanRunCodeLite(Path);
+	return Location.IsValid();
 }
 
 FName FCodeLiteSourceCodeAccessor::GetFName() const
 {
-	return FName("CodeLite");
+	return FName("CodeLiteSourceCodeAccessor");
 }
 
 FText FCodeLiteSourceCodeAccessor::GetNameText() const
 {
-	return LOCTEXT("CodeLiteDisplayName", "CodeLite 7/8.x");
+	return LOCTEXT("CodeLiteDisplayName", "CodeLite 10/11/12");
 }
 
 FText FCodeLiteSourceCodeAccessor::GetDescriptionText() const
@@ -44,88 +72,72 @@ FText FCodeLiteSourceCodeAccessor::GetDescriptionText() const
 
 bool FCodeLiteSourceCodeAccessor::OpenSolution()
 {
-	FString SolutionPath = GetSolutionPath();
-	return OpenSolutionAtPath(SolutionPath);
+	if (!Location.IsValid())
+	{
+		return false;
+	}
+
+	return OpenSolutionAtPath(GetSolutionPath());
 }
 
 bool FCodeLiteSourceCodeAccessor::OpenSolutionAtPath(const FString& InSolutionPath)
 {
-	FString SolutionPath = InSolutionPath;
-	if (!SolutionPath.EndsWith(TEXT(".workspace")))
+	if (!Location.IsValid())
 	{
-		SolutionPath += TEXT(".workspace");
-	}
-
-	FString CodeLitePath;
-	if(!CanRunCodeLite(CodeLitePath))
-	{
-		UE_LOG(LogCodeLiteAccessor, Warning, TEXT("FCodeLiteSourceCodeAccessor::OpenSolution: Cannot find CodeLite binary"));
 		return false;
 	}
 
-	FString Args = TEXT("\"") + SolutionPath.ReplaceQuotesWithEscapedQuotes() + TEXT("\"");
+	FString SolutionPath = InSolutionPath;
 
-	UE_LOG(LogCodeLiteAccessor, Warning, TEXT("FCodeLiteSourceCodeAccessor::OpenSolution: \"%s\" %s"), *CodeLitePath, *Args);
-	
-	FProcHandle Proc = FPlatformProcess::CreateProc(*CodeLitePath, *Args, true, false, false, nullptr, 0, nullptr, nullptr);
-	if(Proc.IsValid())
+	if (!SolutionPath.EndsWith(TEXT(".workspace")))
 	{
-		FPlatformProcess::CloseProc(Proc);
-		return true;
+		SolutionPath = SolutionPath + TEXT(".workspace");
 	}
-	return false;
+
+	TArray<FString> Args;
+	Args.Add(MakePath(SolutionPath));
+	return Launch(Args);
 }
 
 bool FCodeLiteSourceCodeAccessor::DoesSolutionExist() const
  {
-	FString SolutionPath = GetSolutionPath();	
-	return FPaths::FileExists(SolutionPath);
+	return FPaths::FileExists(GetSolutionPath());
  }
 
 bool FCodeLiteSourceCodeAccessor::OpenSourceFiles(const TArray<FString>& AbsoluteSourcePaths)
 {
-	FString CodeLitePath;
-	if(!CanRunCodeLite(CodeLitePath))
+	if (!Location.IsValid())
 	{
-		UE_LOG(LogCodeLiteAccessor, Warning, TEXT("FCodeLiteSourceCodeAccessor::OpenSourceFiles: Cannot find CodeLite binary"));
 		return false;
 	}
 
-	for(const auto& SourcePath : AbsoluteSourcePaths)
-	{
-		const FString Path = FString::Printf(TEXT("\"%s\""), *SourcePath);
+	FString SolutionDir = GetSolutionPath();
+	TArray<FString> Args;
+	Args.Add(MakePath(SolutionDir));
 
-		FProcHandle Proc = FPlatformProcess::CreateProc(*CodeLitePath, *Path, true, false, false, nullptr, 0, nullptr, nullptr);
-		if(Proc.IsValid())
-		{
-			UE_LOG(LogCodeLiteAccessor, Warning, TEXT("CodeLiteSourceCodeAccessor::OpenSourceFiles: %s"), *Path);
-			FPlatformProcess::CloseProc(Proc);
-			return true;
-		}
+	for (const FString& SourcePath : AbsoluteSourcePaths)
+	{
+		Args.Add(MakePath(SourcePath));
 	}
 
-	return false;
+	return Launch(Args);
 }
 
 bool FCodeLiteSourceCodeAccessor::OpenFileAtLine(const FString& FullPath, int32 LineNumber, int32 ColumnNumber)
 {
-	FString CodeLitePath;
-	if(!CanRunCodeLite(CodeLitePath))
+	if (!Location.IsValid())
 	{
-		UE_LOG(LogCodeLiteAccessor,Warning, TEXT("FCodeLiteSourceCodeAccessor::OpenFileAtLine: Cannot find CodeLite binary"));
 		return false;
 	}
 
-	const FString Path = FString::Printf(TEXT("\"%s --line=%d\""), *FullPath, LineNumber);
+	LineNumber = LineNumber > 0 ? LineNumber : 1;
 
-	if(FPlatformProcess::CreateProc(*CodeLitePath, *Path, true, true, false, nullptr, 0, nullptr, nullptr).IsValid())
-	{
-		UE_LOG(LogCodeLiteAccessor, Warning, TEXT("FCodeLiteSourceCodeAccessor::OpenSolution: Cannot find CodeLite binary"));
-	}
+	FString SolutionDir = GetSolutionPath();
+	TArray<FString> Args;
+	Args.Add(MakePath(SolutionDir));
+	Args.Add(MakePath(FullPath) + FString::Printf(TEXT(" --line=%d"), LineNumber));
 
-	UE_LOG(LogCodeLiteAccessor, Warning, TEXT("CodeLiteSourceCodeAccessor::OpenFileAtLine: %s %d"), *FullPath, LineNumber);
-
-	return true;
+	return Launch(Args);
 }
 
 bool FCodeLiteSourceCodeAccessor::AddSourceFiles(const TArray<FString>& AbsoluteSourcePaths, const TArray<FString>& AvailableModules)
@@ -145,97 +157,64 @@ void FCodeLiteSourceCodeAccessor::Tick(const float DeltaTime)
 	// TODO What can we do here?
 }
 
-
-bool FCodeLiteSourceCodeAccessor::CanRunCodeLite(FString& OutPath) const
-{
-	// TODO This might be not a good idea to find an executable.
-	OutPath = TEXT("/usr/bin/codelite");
-	return FPaths::FileExists(OutPath);
-}
-
-
-bool FCodeLiteSourceCodeAccessor::IsIDERunning()
-{
-#if PLATFORM_LINUX
-	pid_t pid = FindProcess("codelite");
-	if(pid == -1)
-	{
-		return false;
-	}
-
-	return true;
-#else
-	return false;
-#endif
-}
-
 FString FCodeLiteSourceCodeAccessor::GetSolutionPath() const
 {
-	if(IsInGameThread())
+	FScopeLock Lock(&CachedSolutionPathCriticalSection);
+
+	if (IsInGameThread())
 	{
 		CachedSolutionPath = FPaths::ConvertRelativePathToFull(FPaths::ProjectDir());
-		
-		if (!FUProjectDictionary::GetDefault().IsForeignProject(CachedSolutionPath))
-		{
-			CachedSolutionPath = FPaths::Combine(FPaths::RootDir(), TEXT("UnrealEditor.workspace"));
-		}
-		else
+
+		FString PrimaryProjectPath;
+		if (!FFileHelper::LoadFileToString(PrimaryProjectPath, *(FPaths::EngineIntermediateDir() / TEXT("ProjectFiles") / TEXT("PrimaryProjectPath.txt"))) || 
+			(PrimaryProjectPath.TrimStartAndEndInline(), PrimaryProjectPath.IsEmpty()))
 		{
-			FString BaseName = FApp::HasProjectName() ? FApp::GetProjectName() : FPaths::GetBaseFilename(CachedSolutionPath);
-			CachedSolutionPath = FPaths::Combine(CachedSolutionPath, BaseName + TEXT(".workspace"));
+			const FProjectDescriptor* CurrentProject = IProjectManager::Get().GetCurrentProject();
+			const FString ProjectDir = FPaths::ConvertRelativePathToFull(FPaths::ProjectDir());
+
+			if ((CurrentProject == nullptr || CurrentProject->Modules.Num() == 0) || !FUProjectDictionary::GetDefault().IsForeignProject(ProjectDir))
+			{
+				PrimaryProjectPath = FPaths::RootDir() / TEXT("UE5");
+			}
+			else
+			{
+				const FString BaseName = FApp::HasProjectName() ? FApp::GetProjectName() : FPaths::GetBaseFilename(ProjectDir);
+				PrimaryProjectPath = ProjectDir / BaseName;
+			}
 		}
+
+		CachedSolutionPath = PrimaryProjectPath + TEXT(".workspace");
+		FPaths::NormalizeFilename(CachedSolutionPath);
 	}
+
 	return CachedSolutionPath;
 }
 
-#if PLATFORM_LINUX
-// This needs to be changed to use platform abstraction layer (see FPlatformProcess::FProcEnumerator)
-pid_t FCodeLiteSourceCodeAccessor::FindProcess(const char* name)
+bool FCodeLiteSourceCodeAccessor::Launch(const TArray<FString>& InArgs)
 {
-	// TODO This is only to test. Will be changed later.
-	DIR* dir;
-	struct dirent* ent;
-	char* endptr;
-	char buf[512];
-
-	if(!(dir = opendir("/proc")))
-	{
-		perror("can't open /proc");
-		return -1;
-	}
-
-	while((ent = readdir(dir)) != NULL)
+	if (Location.IsValid())
 	{
-		long lpid = strtol(ent->d_name, &endptr, 10);
-		if(*endptr != '\0')
+		FString ArgsString;
+		for (const FString& Arg : InArgs)
 		{
-			continue;
+			ArgsString.Append(Arg);
+			ArgsString.Append(TEXT(" "));
 		}
 
-		// Try to open the cmdline file
-		snprintf(buf, sizeof(buf), "/proc/%ld/cmdline", lpid);
-		FILE* fp = fopen(buf, "r");
-		if(fp)
+		uint32 ProcessID;
+		FProcHandle hProcess = FPlatformProcess::CreateProc(*Location.URL, *ArgsString, true, false, false, &ProcessID, 0, nullptr, nullptr, nullptr);
+		if (hProcess.IsValid())
 		{
-			if(fgets(buf, sizeof(buf), fp) != NULL)
-			{
-				// check the first token in the file, the program name.
-				char* first = strtok(buf, " ");
-				if(!strcmp(first, name))
-				{
-					fclose(fp);
-					closedir(dir);
-					return (pid_t)lpid;
-				}
-			}
-			fclose(fp);
+			FPlatformProcess::CloseProc(hProcess);
+			return true;
+		}
+		else
+		{
+			return false;
 		}
 	}
 
-	closedir(dir);
-
-	return -1;
+	return false;
 }
-#endif
 
 #undef LOCTEXT_NAMESPACE
diff --git a/Engine/Plugins/Developer/CodeLiteSourceCodeAccess/Source/CodeLiteSourceCodeAccess/Private/CodeLiteSourceCodeAccessor.h b/Engine/Plugins/Developer/CodeLiteSourceCodeAccess/Source/CodeLiteSourceCodeAccess/Private/CodeLiteSourceCodeAccessor.h
--- a/Engine/Plugins/Developer/CodeLiteSourceCodeAccess/Source/CodeLiteSourceCodeAccess/Private/CodeLiteSourceCodeAccessor.h
+++ b/Engine/Plugins/Developer/CodeLiteSourceCodeAccess/Source/CodeLiteSourceCodeAccess/Private/CodeLiteSourceCodeAccessor.h
@@ -11,7 +11,7 @@ public:
 	/**
 	 * Call if we've likely modified the availability of the source code accessor. 
 	 */
-	virtual void RefreshAvailability() override {}
+	virtual void RefreshAvailability() override;
 
 	/**
 	 * Can we access source code.
@@ -86,21 +86,27 @@ public:
 
 private:
 
+	struct FLocation
+	{
+		bool IsValid() const
+		{
+			return URL.Len() > 0;
+		}
+
+		FString URL;
+	};
+
+	/** Location instance */
+	FLocation Location;
+
 	/** String storing the solution path obtained from the module manager to avoid having to use it on a thread */
 	mutable FString CachedSolutionPath;
 
-	/** Tests if CodeLite is present and returns path to it */
-	bool CanRunCodeLite(FString& OutPath) const;
-
-	/** Check if CodeLite is already running */
-	bool IsIDERunning();
+	/** Critical section for updating SolutionPath */
+	mutable FCriticalSection CachedSolutionPathCriticalSection;
 
 	/** Gets solution path */
 	FString GetSolutionPath() const;
 
-	// TODO Well, change that as soon as possible.
-#if PLATFORM_LINUX
-	pid_t FindProcess(const char* name);
-#endif
-
+	bool Launch(const TArray<FString>& InArgs);
 };
